SourceControl: Added SourceControlCheckInPrompter module, which periodically reminds the user that submitting their files to SourceControl is a good idea to avoid losing work.

#rnx

diff --git a/Engine/Source/Developer/SourceControlCheckInPrompt/Private/SourceControlCheckInPromptModule.cpp b/Engine/Source/Developer/SourceControlCheckInPrompt/Private/SourceControlCheckInPromptModule.cpp
new file mode 100644
--- /dev/null
+++ b/Engine/Source/Developer/SourceControlCheckInPrompt/Private/SourceControlCheckInPromptModule.cpp
@@ -0,0 +1,117 @@
+// Copyright Epic Games, Inc. All Rights Reserved.
+
+#include "SourceControlCheckInPromptModule.h"
+#include "SourceControlCheckInPrompter.h"
+#include "SourceControlWindows.h"
+#include "Dialogs/Dialogs.h"
+#include "Widgets/Notifications/SNotificationList.h"
+#include "Framework/Notifications/NotificationManager.h"
+#include "RevisionControlStyle/RevisionControlStyle.h"
+#include "ISourceControlModule.h"
+
+#define LOCTEXT_NAMESPACE "SourceControlCheckInPromptModule"
+
+FSourceControlCheckInPromptModule::FSourceControlCheckInPromptModule()
+{
+}
+
+void FSourceControlCheckInPromptModule::StartupModule()
+{
+	SourceControlCheckInPrompter = MakeShared<FSourceControlCheckInPrompter>();
+	SourceControlCheckInPrompter->Init();
+
+	ProviderChangedHandle = ISourceControlModule::Get().RegisterProviderChanged(
+		FSourceControlProviderChanged::FDelegate::CreateRaw(this, &FSourceControlCheckInPromptModule::OnProviderChanged)
+	);
+}
+
+void FSourceControlCheckInPromptModule::ShutdownModule()
+{
+	if (ProviderChangedHandle.IsValid())
+	{
+		ISourceControlModule::Get().UnregisterProviderChanged(ProviderChangedHandle);
+		ProviderChangedHandle.Reset();
+	}
+	
+	SourceControlCheckInPrompter.Reset();
+}
+
+void FSourceControlCheckInPromptModule::ShowModal(const FText& InMessage)
+{
+	FText Title = LOCTEXT("CheckInModal_Title", "Check-in Changes");
+
+	FSuppressableWarningDialog::FSetupInfo Info(InMessage, Title, "CheckInModal_Warning");
+	Info.ConfirmText = LOCTEXT("CheckInModal_Yes", "Check-in Changes");
+	Info.CancelText = LOCTEXT("CheckInModal_No", "Skip Check-in");
+	Info.CheckBoxText = FText::GetEmpty(); // Not suppressable!
+	Info.Image = const_cast<FSlateBrush*>(FRevisionControlStyleManager::Get().GetBrush("RevisionControl.StatusBar.HasLocalChanges"));
+
+	FSuppressableWarningDialog CheckInWarning(Info);
+	if (CheckInWarning.ShowModal() == FSuppressableWarningDialog::Confirm)
+	{
+		FSourceControlWindows::ChoosePackagesToCheckIn();
+	}
+}
+
+void FSourceControlCheckInPromptModule::ShowToast(const FText& InMessage)
+{
+	if (TSharedPtr<SNotificationItem> CheckInNotificationPin = CheckInNotification.Pin())
+	{
+		CheckInNotificationPin->SetCompletionState(SNotificationItem::CS_None);
+		CheckInNotificationPin->ExpireAndFadeout();
+	}
+
+	FNotificationInfo Info(InMessage);
+	Info.bFireAndForget = false;
+	Info.bUseThrobber = false;
+	Info.ExpireDuration = 0.0f;
+	Info.FadeOutDuration = 0.0f;
+	Info.ButtonDetails.Add(FNotificationButtonInfo(LOCTEXT("CheckInToast_CheckIn", "Check In"), LOCTEXT("CheckInToast_CheckIn_Tooltip", "Open the submit files dialog to check in your work"), FSimpleDelegate::CreateRaw(this, &FSourceControlCheckInPromptModule::OnNotificationCheckInClicked)));
+	Info.ButtonDetails.Add(FNotificationButtonInfo(LOCTEXT("CheckInToast_Dismiss", "Dismiss"), LOCTEXT("CheckInToast_Dismiss_Tooltip", "Dismiss this notification"), FSimpleDelegate::CreateRaw(this, &FSourceControlCheckInPromptModule::OnNotificationDismissClicked)));
+	Info.Image = FRevisionControlStyleManager::Get().GetBrush("RevisionControl.StatusBar.HasLocalChanges");
+	CheckInNotification = FSlateNotificationManager::Get().AddNotification(Info);
+
+	if (TSharedPtr<SNotificationItem> CheckInNotificationPin = CheckInNotification.Pin())
+	{
+		CheckInNotificationPin->SetCompletionState(SNotificationItem::CS_Pending);
+	}
+}
+
+void FSourceControlCheckInPromptModule::OnNotificationCheckInClicked()
+{
+	if (TSharedPtr<SNotificationItem> CheckInNotificationPin = CheckInNotification.Pin())
+	{
+		CheckInNotificationPin->SetCompletionState(SNotificationItem::CS_None);
+		CheckInNotificationPin->SetExpireDuration(0.0f);
+		CheckInNotificationPin->SetFadeOutDuration(0.0f);
+		CheckInNotificationPin->ExpireAndFadeout();
+	}
+
+	FSourceControlWindows::ChoosePackagesToCheckIn();
+}
+
+void FSourceControlCheckInPromptModule::OnNotificationDismissClicked()
+{
+	if (TSharedPtr<SNotificationItem> CheckInNotificationPin = CheckInNotification.Pin())
+	{
+		CheckInNotificationPin->SetCompletionState(SNotificationItem::CS_None);
+		CheckInNotificationPin->SetExpireDuration(0.0f);
+		CheckInNotificationPin->SetFadeOutDuration(0.0f);
+		CheckInNotificationPin->ExpireAndFadeout();
+	}
+}
+
+void FSourceControlCheckInPromptModule::OnProviderChanged(ISourceControlProvider& OldProvider, ISourceControlProvider& NewProvider)
+{
+	if (TSharedPtr<SNotificationItem> CheckInNotificationPin = CheckInNotification.Pin())
+	{
+		CheckInNotificationPin->SetCompletionState(SNotificationItem::CS_None);
+		CheckInNotificationPin->SetExpireDuration(0.0f);
+		CheckInNotificationPin->SetFadeOutDuration(0.0f);
+		CheckInNotificationPin->ExpireAndFadeout();
+	}
+}
+
+IMPLEMENT_MODULE( FSourceControlCheckInPromptModule, SourceControlCheckInPrompt );
+
+#undef LOCTEXT_NAMESPACE
\ No newline at end of file
diff --git a/Engine/Source/Developer/SourceControlCheckInPrompt/Private/SourceControlCheckInPromptModule.h b/Engine/Source/Developer/SourceControlCheckInPrompt/Private/SourceControlCheckInPromptModule.h
new file mode 100644
--- /dev/null
+++ b/Engine/Source/Developer/SourceControlCheckInPrompt/Private/SourceControlCheckInPromptModule.h
@@ -0,0 +1,56 @@
+// Copyright Epic Games, Inc. All Rights Reserved.
+
+#pragma once
+
+#include "CoreMinimal.h"
+#include "Modules/ModuleManager.h"
+#include "ISourceControlCheckInPromptModule.h"
+
+class FSourceControlCheckInPrompter;
+class ISourceControlProvider;
+class SNotificationItem;
+
+class FSourceControlCheckInPromptModule : public ISourceControlCheckInPromptModule
+{
+public:
+	FSourceControlCheckInPromptModule();
+
+	/** IModuleInterface implementation */
+	virtual void StartupModule() override;
+	virtual void ShutdownModule() override;
+
+	/** ISourceControlCheckInPromptModule implementation */
+	virtual void ShowModal(const FText& InMessage) override;
+	virtual void ShowToast(const FText& InMessage) override;
+
+	/**
+	 * Gets a reference to the source control check in prompt module instance.
+	 *
+	 * @return A reference to the source control check in prompt module.
+	 */
+	static FSourceControlCheckInPromptModule& Get()
+	{
+		return FModuleManager::LoadModuleChecked<FSourceControlCheckInPromptModule>("SourceControlCheckInPrompt");
+	}
+
+	static FSourceControlCheckInPromptModule* TryGet()
+	{
+		return FModuleManager::GetModulePtr<FSourceControlCheckInPromptModule>("SourceControlCheckInPrompt");
+	}
+
+protected:
+	void OnNotificationCheckInClicked();
+	void OnNotificationDismissClicked();
+
+	void OnProviderChanged(ISourceControlProvider& OldProvider, ISourceControlProvider& NewProvider);
+
+private:
+	/** The prompter which triggers a periodic prompt */
+	TSharedPtr<FSourceControlCheckInPrompter> SourceControlCheckInPrompter;
+
+	/** The notification toast that a check-in is recommended */
+	TWeakPtr<SNotificationItem> CheckInNotification;
+
+	/** The delegate handle for provider changes */
+	FDelegateHandle ProviderChangedHandle;
+};
diff --git a/Engine/Source/Developer/SourceControlCheckInPrompt/Private/SourceControlCheckInPrompter.cpp b/Engine/Source/Developer/SourceControlCheckInPrompt/Private/SourceControlCheckInPrompter.cpp
new file mode 100644
--- /dev/null
+++ b/Engine/Source/Developer/SourceControlCheckInPrompt/Private/SourceControlCheckInPrompter.cpp
@@ -0,0 +1,212 @@
+// Copyright Epic Games, Inc. All Rights Reserved.
+
+#include "SourceControlCheckInPrompter.h"
+#include "SourceControlCheckInPromptModule.h"
+#include "SourceControlHelpers.h"
+#include "SourceControlOperations.h"
+#include "ISourceControlModule.h"
+#include "Containers/Ticker.h"
+#include "UObject/ObjectSaveContext.h"
+#include "UObject/Package.h"
+#include "HAL/IConsoleManager.h"
+#include "Editor.h"
+
+#define LOCTEXT_NAMESPACE "SourceControlCheckInPrompter"
+
+static const FTimespan IntervalNoCheckins = FTimespan::FromDays(1);
+static const FTimespan IntervalBetweenPrompts = FTimespan::FromDays(1);
+static const FTimespan IntervalBetweenGetSubmittedChangelists = FTimespan::FromMinutes(10);
+
+TAutoConsoleVariable<bool> CVarSourceControlEnablePeriodicCheckInPrompt(
+	TEXT("SourceControl.CheckInPrompt.EnablePeriodic"),
+	false,
+	TEXT("Enables a periodic check-in prompt, reminding the user to check-in to avoid losing work."),
+	ECVF_Default);
+
+static FString GetEditorMapName()
+{
+	if (UWorld* EditorWorld = GEditor->GetEditorWorldContext().World())
+	{
+		if (UPackage* EditorWorldPackage = EditorWorld->GetPackage())
+		{
+			const FString EditorWorldPackageName = EditorWorldPackage->GetName();
+			if (!EditorWorldPackageName.StartsWith(TEXT("/Temp/")))
+			{
+				return EditorWorldPackageName;
+			}
+		}
+	}
+
+	return FString();
+}
+
+FSourceControlCheckInPrompter::FSourceControlCheckInPrompter()
+	: PromptFlowMapName()
+	, TimeCheckInPromptShown()
+	, TimeGetSubmittedChangelistsExecuted()
+{
+}
+
+FSourceControlCheckInPrompter::~FSourceControlCheckInPrompter()
+{
+	UPackage::PackageSavedWithContextEvent.RemoveAll(this);
+}
+
+void FSourceControlCheckInPrompter::Init()
+{
+	UPackage::PackageSavedWithContextEvent.AddRaw(this, &FSourceControlCheckInPrompter::OnPackageSaved);
+}
+
+// Step 1: Initiate the periodic prompt flow whenever a package is saved.
+//         Trigger a SourceControl operation to determine how much changes the user submitted in the past day.
+void FSourceControlCheckInPrompter::OnPackageSaved(const FString& Filename, UPackage* Pkg, FObjectPostSaveContext ObjectSaveContext)
+{
+	FString EditorMapName = GetEditorMapName();
+	if (PromptFlowMapName.IsEmpty() && !EditorMapName.IsEmpty())
+	{
+		bool bPromptEnabled = CVarSourceControlEnablePeriodicCheckInPrompt.GetValueOnGameThread();
+		bool bIsPromptAllowed = IsPromptAllowed();
+		bool bIsGetSubmittedChangelistsAllowed = IsGetSubmittedChangelistsAllowed();
+		if (bPromptEnabled && bIsPromptAllowed && bIsGetSubmittedChangelistsAllowed)
+		{
+			// Execute it with a filter that looks for submitted changelist for the current user in the last day.
+			TSharedRef<FGetSubmittedChangelists> Operation = ISourceControlOperation::Create<FGetSubmittedChangelists>();
+			Operation->SetDateToFilter(FDateTime::UtcNow());
+			Operation->SetDateFromFilter(FDateTime::UtcNow() - IntervalNoCheckins);
+			Operation->SetOwnedFilter(true);
+
+			ISourceControlProvider& SourceControlProvider = ISourceControlModule::Get().GetProvider();
+			if (SourceControlProvider.IsAvailable() && SourceControlProvider.CanExecuteOperation(Operation))
+			{
+				// Start flow.
+				PromptFlowMapName = EditorMapName;
+
+				// Update the time the get submitted changelists operation was last executed.
+				TimeGetSubmittedChangelistsExecuted.Add(EditorMapName, FDateTime::Now());
+
+				// Execute it asynchronously.
+				SourceControlProvider.Execute(Operation, EConcurrency::Asynchronous,
+					FSourceControlOperationComplete::CreateSP(this, &FSourceControlCheckInPrompter::OnSourceControlOperationComplete)
+				);
+			}
+		}
+	}
+}
+
+// Step 2: When the SourceControl operation completes, check if any check-ins were found.
+//         If yes, end flow.
+//         If no, continue flow by showing the prompt.
+void FSourceControlCheckInPrompter::OnSourceControlOperationComplete(const FSourceControlOperationRef& InOperation, ECommandResult::Type InResult)
+{
+	ISourceControlProvider& SourceControlProvider = ISourceControlModule::Get().GetProvider();
+
+	if (InOperation->GetName() == TEXT("GetSubmittedChangelists"))
+	{
+		TSharedRef<FGetSubmittedChangelists, ESPMode::ThreadSafe> Operation = StaticCastSharedRef<FGetSubmittedChangelists>(InOperation);
+
+		const TArray<FSourceControlChangelistRef> SubmittedChangelists = Operation->GetSubmittedChangelists();
+		if (SubmittedChangelists.Num() == 0)
+		{
+			// Continue flow.
+			// Prompt user as he's got no recent check-ins.
+			FTSTicker::GetCoreTicker().AddTicker(
+				FTickerDelegate::CreateSP(this, &FSourceControlCheckInPrompter::OnAttemptPrompt), 1.0f
+			);
+		}
+		else
+		{
+			PromptFlowMapName.Empty();
+		}
+	}
+}
+
+// Step 3: When it's determined that a prompt should be shown, wait for a suitable moment in the editor.
+bool FSourceControlCheckInPrompter::OnAttemptPrompt(float)
+{
+	// Abort if map changed, project closed, etc...
+	FString EditorMapName = GetEditorMapName();
+	if (PromptFlowMapName != EditorMapName)
+	{
+		PromptFlowMapName.Empty();
+		return false;
+	}
+
+	// Update the time the prompt was last shown.
+	TimeCheckInPromptShown.Add(EditorMapName, FDateTime::Now());
+
+	// Show the prompt.
+	FSourceControlCheckInPromptModule::Get().ShowToast(
+		LOCTEXT("MessageToast", "It's been more than 24 hours since you checked in. Check in to revision control to back up your project?")
+	);
+
+	PromptFlowMapName.Empty();
+	return false;
+}
+
+bool FSourceControlCheckInPrompter::IsPromptAllowed() const
+{
+	// Ensure the SourceControl system is available.
+	ISourceControlProvider& SourceControlProvider = ISourceControlModule::Get().GetProvider();
+	if (!SourceControlProvider.IsAvailable())
+	{
+		return false;
+	}
+
+	// Ensure it has something to check-in.
+	if (SourceControlProvider.GetNumLocalChanges().IsSet())
+	{
+		int NumLocalChanges = SourceControlProvider.GetNumLocalChanges().GetValue();
+		if (NumLocalChanges == 0)
+		{
+			return false;
+		}
+	}
+
+	// Ensure there's a world in the editor.
+	UWorld* EditorWorld = GEditor->GetEditorWorldContext().World();
+	if (EditorWorld == nullptr)
+	{
+		return false;
+	}
+
+	// Ensure the prompt hasn't been shown too recently for that world.
+	FString PackageName = EditorWorld->GetPackage()->GetName();
+	if (const FDateTime* LastPrompt = TimeCheckInPromptShown.Find(PackageName))
+	{
+		return FDateTime::Now() >= *LastPrompt + IntervalBetweenPrompts;
+	}
+	else
+	{
+		return true; // A prompt wasn't shown before.
+	}
+}
+
+bool FSourceControlCheckInPrompter::IsGetSubmittedChangelistsAllowed() const
+{
+	// Ensure the SourceControl system is available.
+	ISourceControlProvider& SourceControlProvider = ISourceControlModule::Get().GetProvider();
+	if (!SourceControlProvider.IsAvailable())
+	{
+		return false;
+	}
+
+	// Ensure there's a world in the editor.
+	UWorld* EditorWorld = GEditor->GetEditorWorldContext().World();
+	if (EditorWorld == nullptr)
+	{
+		return false;
+	}
+
+	// Ensure the operation isn't executed too often for that world.
+	FString PackageName = EditorWorld->GetPackage()->GetName();
+	if (const FDateTime* LastOperation = TimeGetSubmittedChangelistsExecuted.Find(PackageName))
+	{
+		return FDateTime::Now() >= *LastOperation + IntervalBetweenGetSubmittedChangelists;
+	}
+	else
+	{
+		return true; // An operation wasn't executed yet.
+	}
+}
+
+#undef LOCTEXT_NAMESPACE
\ No newline at end of file
diff --git a/Engine/Source/Developer/SourceControlCheckInPrompt/Private/SourceControlCheckInPrompter.h b/Engine/Source/Developer/SourceControlCheckInPrompt/Private/SourceControlCheckInPrompter.h
new file mode 100644
--- /dev/null
+++ b/Engine/Source/Developer/SourceControlCheckInPrompt/Private/SourceControlCheckInPrompter.h
@@ -0,0 +1,65 @@
+// Copyright Epic Games, Inc. All Rights Reserved.
+
+#pragma once
+
+#include "CoreMinimal.h"
+#include "UObject/WeakObjectPtr.h"
+#include "UObject/WeakObjectPtrTemplates.h"
+#include "ISourceControlProvider.h"
+
+class FObjectPostSaveContext;
+
+class FSourceControlCheckInPrompter : public TSharedFromThis<FSourceControlCheckInPrompter>
+{
+public:
+	FSourceControlCheckInPrompter();
+	virtual ~FSourceControlCheckInPrompter();
+
+	/**
+	 * Initializes the instance.
+	 */
+	void Init();
+
+private:
+	/**
+	 * Called when a package has been saved.
+	 *
+	 * @param Filename The filename the package was saved to
+	 * @param Obj The package that was saved
+	 */
+	void OnPackageSaved(const FString& Filename, UPackage* Pkg, FObjectPostSaveContext ObjectSaveContext);
+
+	/**
+	 * Called when an executed source control operation completes.
+	 */
+	void OnSourceControlOperationComplete(const FSourceControlOperationRef& InOperation, ECommandResult::Type InResult);
+
+	/**
+	 * Called when a periodic check-in prompt should attempted to be shown.
+	 * @param DeltaTime argument is the time since the last game frame
+	 */
+	bool OnAttemptPrompt(float);
+
+	/**
+	 * Checks if it's allowed to display the periodic check-in prompt.
+	 * The prompt is only expected to be shown once every X hours.
+	 */
+	bool IsPromptAllowed() const;
+
+	/**
+	 * Checks if it's allow to perform an asynchronous get submitted changelist operation.
+	 * The operation is only expected to be executed once every X minutes.
+	 */
+	bool IsGetSubmittedChangelistsAllowed() const;
+
+private:
+	/** When non empty, contains the package name of the map the prompt will be shown for */
+	FString PromptFlowMapName;
+
+	/** Time the last prompt was shown by map name */
+	TMap<FString, FDateTime> TimeCheckInPromptShown;
+
+	/** Time the last get submitted changelists operation was performed by map name */
+	TMap<FString, FDateTime> TimeGetSubmittedChangelistsExecuted;
+
+};
\ No newline at end of file
diff --git a/Engine/Source/Developer/SourceControlCheckInPrompt/Public/ISourceControlCheckInPromptModule.h b/Engine/Source/Developer/SourceControlCheckInPrompt/Public/ISourceControlCheckInPromptModule.h
new file mode 100644
--- /dev/null
+++ b/Engine/Source/Developer/SourceControlCheckInPrompt/Public/ISourceControlCheckInPromptModule.h
@@ -0,0 +1,42 @@
+// Copyright Epic Games, Inc. All Rights Reserved.
+
+#pragma once
+
+#include "CoreMinimal.h"
+#include "Modules/ModuleInterface.h"
+#include "Modules/ModuleManager.h"
+#include "Delegates/Delegate.h"
+
+class ISourceControlCheckInPromptModule : public IModuleInterface
+{
+public:
+	/**
+	 * Shows a check-in prompt reminding the user that it's important to frequently save their work.
+	 * @param InMessage		The message to display in the message box
+	 */
+	virtual void ShowModal(const FText& InMessage) = 0;
+
+	/**
+	 * Shows a check-in toast reminder the user that it's important to frequently save their work.
+	 * @param InMessage		The message to display in the toast
+	 */
+	virtual void ShowToast(const FText& InMessage) = 0;
+
+public:
+	/**
+	 * Gets a reference to the source control check in prompt module instance.
+	 *
+	 * @return A reference to the source control check in module.
+	 */
+	static inline ISourceControlCheckInPromptModule& Get()
+	{
+		static FName SourceControlCheckInPromptModule("SourceControlCheckInPrompt");
+		return FModuleManager::LoadModuleChecked<ISourceControlCheckInPromptModule>(SourceControlCheckInPromptModule);
+	}
+
+	static inline ISourceControlCheckInPromptModule* TryGet()
+	{
+		static FName SourceControlCheckInPromptModule("SourceControlCheckInPrompt");
+		return FModuleManager::GetModulePtr<ISourceControlCheckInPromptModule>(SourceControlCheckInPromptModule);
+	}
+};
diff --git a/Engine/Source/Developer/SourceControlCheckInPrompt/SourceControlCheckInPrompt.Build.cs b/Engine/Source/Developer/SourceControlCheckInPrompt/SourceControlCheckInPrompt.Build.cs
new file mode 100644
--- /dev/null
+++ b/Engine/Source/Developer/SourceControlCheckInPrompt/SourceControlCheckInPrompt.Build.cs
@@ -0,0 +1,22 @@
+// Copyright Epic Games, Inc. All Rights Reserved.
+
+using UnrealBuildTool;
+
+public class SourceControlCheckInPrompt : ModuleRules
+{
+	public SourceControlCheckInPrompt(ReadOnlyTargetRules Target) : base(Target)
+	{
+		PrivateDependencyModuleNames.AddRange(
+			new string[] {
+				"Core",
+				"CoreUObject",
+				"Engine",
+				"UnrealEd",
+				"Slate",
+				"SlateCore",
+				"SourceControl",
+				"SourceControlWindows"
+			}
+		);
+	}
+}
