SkeinUE: Fixed 'resolve conflicts' button becoming unresponsive.

What happened? The FSuppressableWarningDialog::ShowModal() function returns SUPPRESSED if the dialog is not shown. This should be treated the same as CONFIRM (in most cases) but it was treated the same as CANCEL.

#rnx

diff --git a/Engine/Source/Developer/SourceControlCheckInPrompt/Private/SourceControlCheckInPromptModule.cpp b/Engine/Source/Developer/SourceControlCheckInPrompt/Private/SourceControlCheckInPromptModule.cpp
--- a/Engine/Source/Developer/SourceControlCheckInPrompt/Private/SourceControlCheckInPromptModule.cpp
+++ b/Engine/Source/Developer/SourceControlCheckInPrompt/Private/SourceControlCheckInPromptModule.cpp
@@ -59,7 +59,7 @@ void FSourceControlCheckInPromptModule::ShowModal(const FText& InMessage)
 	Info.Image = const_cast<FSlateBrush*>(FRevisionControlStyleManager::Get().GetBrush("RevisionControl.StatusBar.HasLocalChanges"));
 
 	FSuppressableWarningDialog CheckInWarning(Info);
-	if (CheckInWarning.ShowModal() == FSuppressableWarningDialog::Confirm)
+	if (CheckInWarning.ShowModal() != FSuppressableWarningDialog::Cancel)
 	{
 		FSourceControlWindows::ChoosePackagesToCheckIn();
 	}
