SkeinUE: Tweak logic for 24hr check in prompt.

diff --git a/Engine/Source/Developer/SourceControlCheckInPrompt/Private/SourceControlCheckInPrompter.cpp b/Engine/Source/Developer/SourceControlCheckInPrompt/Private/SourceControlCheckInPrompter.cpp
--- a/Engine/Source/Developer/SourceControlCheckInPrompt/Private/SourceControlCheckInPrompter.cpp
+++ b/Engine/Source/Developer/SourceControlCheckInPrompt/Private/SourceControlCheckInPrompter.cpp
@@ -41,6 +41,7 @@ FSourceControlCheckInPrompter::FSourceControlCheckInPrompter()
 	: PromptFlowMapName()
 	, TimeCheckInPromptShown()
 	, TimeGetSubmittedChangelistsExecuted()
+	, bPromptDelayed(false)
 {
 }
 
@@ -63,6 +64,12 @@ void FSourceControlCheckInPrompter::Init()
 				{
 					ProjectDirectory = SourceControlProjectDir;
 					ProjectActivationTime = FDateTime::UtcNow();
+					bPromptDelayed = false;
+				}
+
+				if (bPromptDelayed)
+				{
+					CheckPrompt();
 				}
 
 				return true;
@@ -70,39 +77,34 @@ void FSourceControlCheckInPrompter::Init()
 	), 60.f);
 }
 
-// Step 1: Initiate the periodic prompt flow whenever a package is saved.
-//         Trigger a SourceControl operation to determine how much changes the user submitted in the past day.
 void FSourceControlCheckInPrompter::OnPackageSaved(const FString& Filename, UPackage* Pkg, FObjectPostSaveContext ObjectSaveContext)
 {
-	FString EditorMapName = GetEditorMapName();
-	if (PromptFlowMapName.IsEmpty() && !EditorMapName.IsEmpty())
+	CheckPrompt();
+}
+
+// Step 1: Initiate the periodic prompt flow.
+//         Trigger a SourceControl operation to determine how much changes the user submitted in the past day.
+void FSourceControlCheckInPrompter::OnStartPrompt()
+{
+	// Execute it with a filter that looks for submitted changelist for the current user in the last day.
+	TSharedRef<FGetSubmittedChangelists> Operation = ISourceControlOperation::Create<FGetSubmittedChangelists>();
+	Operation->SetDateToFilter(FDateTime::UtcNow());
+	Operation->SetDateFromFilter(FDateTime::UtcNow() - IntervalNoCheckins);
+	Operation->SetOwnedFilter(true);
+
+	ISourceControlProvider& SourceControlProvider = ISourceControlModule::Get().GetProvider();
+	if (SourceControlProvider.IsAvailable() && SourceControlProvider.CanExecuteOperation(Operation))
 	{
-		bool bPromptEnabled = CVarSourceControlEnablePeriodicCheckInPrompt.GetValueOnGameThread();
-		bool bIsPromptAllowed = IsPromptAllowed();
-		bool bIsGetSubmittedChangelistsAllowed = IsGetSubmittedChangelistsAllowed();
-		if (bPromptEnabled && bIsPromptAllowed && bIsGetSubmittedChangelistsAllowed)
-		{
-			// Execute it with a filter that looks for submitted changelist for the current user in the last day.
-			TSharedRef<FGetSubmittedChangelists> Operation = ISourceControlOperation::Create<FGetSubmittedChangelists>();
-			Operation->SetDateToFilter(FDateTime::UtcNow());
-			Operation->SetDateFromFilter(FDateTime::UtcNow() - IntervalNoCheckins);
-			Operation->SetOwnedFilter(true);
-
-			ISourceControlProvider& SourceControlProvider = ISourceControlModule::Get().GetProvider();
-			if (SourceControlProvider.IsAvailable() && SourceControlProvider.CanExecuteOperation(Operation))
-			{
-				// Start flow.
-				PromptFlowMapName = EditorMapName;
+		// Start flow.
+		PromptFlowMapName = GetEditorMapName();
 
-				// Update the time the get submitted changelists operation was last executed.
-				TimeGetSubmittedChangelistsExecuted.Add(EditorMapName, FDateTime::Now());
+		// Update the time the get submitted changelists operation was last executed.
+		TimeGetSubmittedChangelistsExecuted.Add(PromptFlowMapName, FDateTime::Now());
 
-				// Execute it asynchronously.
-				SourceControlProvider.Execute(Operation, EConcurrency::Asynchronous,
-					FSourceControlOperationComplete::CreateSP(this, &FSourceControlCheckInPrompter::OnSourceControlOperationComplete)
-				);
-			}
-		}
+		// Execute it asynchronously.
+		SourceControlProvider.Execute(Operation, EConcurrency::Asynchronous,
+			FSourceControlOperationComplete::CreateSP(this, &FSourceControlCheckInPrompter::OnSourceControlOperationComplete)
+		);
 	}
 }
 
@@ -156,6 +158,51 @@ bool FSourceControlCheckInPrompter::OnAttemptPrompt(float)
 	return false;
 }
 
+void FSourceControlCheckInPrompter::CheckPrompt()
+{
+	FString EditorMapName = GetEditorMapName();
+	if (PromptFlowMapName.IsEmpty() && !EditorMapName.IsEmpty())
+	{
+		bool bPromptEnabled = CVarSourceControlEnablePeriodicCheckInPrompt.GetValueOnGameThread();
+		bool bPromptShow = false;
+
+		if (bPromptEnabled)
+		{
+			if (bPromptDelayed)
+			{
+				bPromptShow = true;
+			}
+			else
+			{
+				bool bIsPromptAllowed = IsPromptAllowed();
+				bool bIsGetSubmittedChangelistsAllowed = IsGetSubmittedChangelistsAllowed();
+				bPromptShow = bIsPromptAllowed && bIsGetSubmittedChangelistsAllowed;
+			}
+		}
+
+		if (bPromptShow)
+		{
+			if (FDateTime::UtcNow() - IntervalSessionLength < ProjectActivationTime)
+			{
+				// As the user has not been active in this world for a sufficient amount of time,
+				// the prompt is delayed until the required interval has passed.
+
+				bPromptShow = false;
+				bPromptDelayed = true;
+			}
+			else
+			{
+				bPromptDelayed = false;
+			}
+		}
+
+		if (bPromptEnabled && bPromptShow)
+		{
+			OnStartPrompt();
+		}
+	}
+}
+
 bool FSourceControlCheckInPrompter::IsPromptAllowed() const
 {
 	// Ensure the SourceControl system is available.
@@ -189,12 +236,6 @@ bool FSourceControlCheckInPrompter::IsPromptAllowed() const
 		return false;
 	}
 
-	// Ensure the user has been active in that world for a sufficient amount of time.
-	if (FDateTime::UtcNow() - IntervalSessionLength < ProjectActivationTime)
-	{
-		return false;
-	}
-
 	// Ensure the prompt hasn't been shown too recently for that world.
 	FString PackageName = EditorWorld->GetPackage()->GetName();
 	if (const FDateTime* LastPrompt = TimeCheckInPromptShown.Find(PackageName))
diff --git a/Engine/Source/Developer/SourceControlCheckInPrompt/Private/SourceControlCheckInPrompter.h b/Engine/Source/Developer/SourceControlCheckInPrompt/Private/SourceControlCheckInPrompter.h
--- a/Engine/Source/Developer/SourceControlCheckInPrompt/Private/SourceControlCheckInPrompter.h
+++ b/Engine/Source/Developer/SourceControlCheckInPrompt/Private/SourceControlCheckInPrompter.h
@@ -29,13 +29,19 @@ private:
 	 */
 	void OnPackageSaved(const FString& Filename, UPackage* Pkg, FObjectPostSaveContext ObjectSaveContext);
 
+	/**
+	 * Called when a prompt flow should start, which might result in one being shown if
+	 * it's been over a day since the last commit.
+	 */
+	void OnStartPrompt();
+
 	/**
 	 * Called when an executed source control operation completes.
 	 */
 	void OnSourceControlOperationComplete(const FSourceControlOperationRef& InOperation, ECommandResult::Type InResult);
 
 	/**
-	 * Called when a periodic check-in prompt should attempted to be shown.
+	 * Called when a check-in prompt should attempted to be shown.
 	 * @param DeltaTime argument is the time since the last game frame
 	 */
 	bool OnAttemptPrompt(float);
@@ -52,6 +58,11 @@ private:
 	 */
 	bool IsGetSubmittedChangelistsAllowed() const;
 
+	/**
+	 * Called when a check should happen whether a prompt should be shown.
+	 */
+	void CheckPrompt();
+
 private:
 	/** When non empty, contains the package name of the map the prompt will be shown for */
 	FString PromptFlowMapName;
@@ -68,4 +79,7 @@ private:
 	/** The project activation time (eg: when was it opened) */
 	FDateTime ProjectActivationTime;
 
+	/** The prompt has been delayed */
+	bool bPromptDelayed;
+
 };
\ No newline at end of file
